- Android build fixes.

diff --git a/Engine/Plugins/Runtime/AndroidMoviePlayer/Source/AndroidMoviePlayer/Private/AndroidMoviePlayer.cpp b/Engine/Plugins/Runtime/AndroidMoviePlayer/Source/AndroidMoviePlayer/Private/AndroidMoviePlayer.cpp
--- a/Engine/Plugins/Runtime/AndroidMoviePlayer/Source/AndroidMoviePlayer/Private/AndroidMoviePlayer.cpp
+++ b/Engine/Plugins/Runtime/AndroidMoviePlayer/Source/AndroidMoviePlayer/Private/AndroidMoviePlayer.cpp
@@ -9,7 +9,7 @@ TSharedPtr<FAndroidMediaPlayerStreamer> AndroidMovieStreamer;
 class FAndroidMoviePlayerModule : public IModuleInterface
 {
 	/** IModuleInterface implementation */
-	virtual void StartupModule() OVERRIDE
+	virtual void StartupModule() override
 	{
 		if (IsSupported())
 		{
@@ -19,7 +19,7 @@ class FAndroidMoviePlayerModule : public IModuleInterface
 		}
 	}
 
-	virtual void ShutdownModule() OVERRIDE
+	virtual void ShutdownModule() override
 	{
 		if (IsSupported())
 		{
diff --git a/Engine/Plugins/Runtime/AndroidMoviePlayer/Source/AndroidMoviePlayer/Private/AndroidMovieStreamer.cpp b/Engine/Plugins/Runtime/AndroidMoviePlayer/Source/AndroidMoviePlayer/Private/AndroidMovieStreamer.cpp
--- a/Engine/Plugins/Runtime/AndroidMoviePlayer/Source/AndroidMoviePlayer/Private/AndroidMovieStreamer.cpp
+++ b/Engine/Plugins/Runtime/AndroidMoviePlayer/Source/AndroidMoviePlayer/Private/AndroidMovieStreamer.cpp
@@ -6,7 +6,7 @@
 #include "AndroidJava.h"
 #include "AndroidFile.h"
 
-#include "Slate.h"
+#include "SlateBasics.h"
 #include "RenderingCommon.h"
 #include "RenderUtils.h"
 #include "Slate/SlateTextures.h"
diff --git a/Engine/Plugins/Runtime/AndroidMoviePlayer/Source/AndroidMoviePlayer/Private/AndroidMovieStreamer.h b/Engine/Plugins/Runtime/AndroidMoviePlayer/Source/AndroidMoviePlayer/Private/AndroidMovieStreamer.h
--- a/Engine/Plugins/Runtime/AndroidMoviePlayer/Source/AndroidMoviePlayer/Private/AndroidMovieStreamer.h
+++ b/Engine/Plugins/Runtime/AndroidMoviePlayer/Source/AndroidMoviePlayer/Private/AndroidMovieStreamer.h
@@ -4,7 +4,7 @@
 
 #include "MoviePlayer.h"
 #include "AndroidJavaMediaPlayer.h"
-#include "Slate.h"
+#include "SlateBasics.h"
 #include "Slate/SlateTextures.h"
 
 // The actual streamer class
