Clear texture for movie playback on Android
#ue-14554
#ue4
#android

diff --git a/Engine/Plugins/Runtime/AndroidMoviePlayer/Source/AndroidMoviePlayer/Private/AndroidMovieStreamer.cpp b/Engine/Plugins/Runtime/AndroidMoviePlayer/Source/AndroidMoviePlayer/Private/AndroidMovieStreamer.cpp
--- a/Engine/Plugins/Runtime/AndroidMoviePlayer/Source/AndroidMoviePlayer/Private/AndroidMovieStreamer.cpp
+++ b/Engine/Plugins/Runtime/AndroidMoviePlayer/Source/AndroidMoviePlayer/Private/AndroidMovieStreamer.cpp
@@ -174,10 +174,19 @@ bool FAndroidMediaPlayerStreamer::StartNextMovie()
 			TexCreate_RenderTargetable,
 			true));
 
-		ENQUEUE_UNIQUE_RENDER_COMMAND_ONEPARAMETER(InitMovieTexture,
+		uint32 FrameBytes = VideoDimensions.X * VideoDimensions.Y * GPixelFormats[PF_B8G8R8A8].BlockBytes;
+
+		ENQUEUE_UNIQUE_RENDER_COMMAND_TWOPARAMETER(InitMovieTexture,
 			FSlateTexture2DRHIRef*, TextureRHIRef, Texture.Get(),
+			uint32, Bytes, FrameBytes,
 			{
 				TextureRHIRef->InitResource();
+	
+				// clear texture to black
+				uint32 Stride = 0;
+				void* TextureBuffer = RHILockTexture2D(TextureRHIRef->GetTypedResource(), 0, RLM_WriteOnly, Stride, false);
+				FMemory::Memset(TextureBuffer, 0, Bytes);
+				RHIUnlockTexture2D(TextureRHIRef->GetTypedResource(), 0, false);
 			});
 
 		MovieViewport->SetTexture(Texture);
