Plugins: Separate categories for startup movie player plug-ins

diff --git a/Engine/Plugins/Runtime/AndroidMoviePlayer/AndroidMoviePlayer.uplugin b/Engine/Plugins/Runtime/AndroidMoviePlayer/AndroidMoviePlayer.uplugin
--- a/Engine/Plugins/Runtime/AndroidMoviePlayer/AndroidMoviePlayer.uplugin
+++ b/Engine/Plugins/Runtime/AndroidMoviePlayer/AndroidMoviePlayer.uplugin
@@ -8,7 +8,7 @@
 	"CreatedByURL" : "http://epicgames.com",
 	"EngineVersion" : 1579795,
 	"Description" : "Android Platform Movie Player using Android Media library",
-	"Category" : "Media Players",
+	"Category" : "Movie Players",
 	"EnabledByDefault" : true,
 	
     "Modules" :
diff --git a/Engine/Plugins/Runtime/AppleMoviePlayer/AppleMoviePlayer.uplugin b/Engine/Plugins/Runtime/AppleMoviePlayer/AppleMoviePlayer.uplugin
--- a/Engine/Plugins/Runtime/AppleMoviePlayer/AppleMoviePlayer.uplugin
+++ b/Engine/Plugins/Runtime/AppleMoviePlayer/AppleMoviePlayer.uplugin
@@ -8,7 +8,7 @@
 	"CreatedByURL" : "http://epicgames.com",
 	"EngineVersion" : "4.2.0",
 	"Description" : "Apple Platform Movie Player using AVPlayer library",
-	"Category" : "Media Players",
+	"Category" : "Movie Players",
 	"EnabledByDefault" : true,
 	
 	"Modules" :
diff --git a/Engine/Plugins/Runtime/WindowsMoviePlayer/WindowsMoviePlayer.uplugin b/Engine/Plugins/Runtime/WindowsMoviePlayer/WindowsMoviePlayer.uplugin
--- a/Engine/Plugins/Runtime/WindowsMoviePlayer/WindowsMoviePlayer.uplugin
+++ b/Engine/Plugins/Runtime/WindowsMoviePlayer/WindowsMoviePlayer.uplugin
@@ -8,7 +8,7 @@
 	"CreatedByURL" : "http://epicgames.com",
 	"EngineVersion" : "4.2.0",
 	"Description" : "Windows Specific Movie Player using Media Foundation",
-	"Category" : "Media Players",
+	"Category" : "Movie Players",
 	"EnabledByDefault" : true,
 	
 	"Modules" :
