Fix startup movies for Android and iOS
- support RHIThread for startup movie playback on Android
#ue4
#android
#ios

diff --git a/Engine/Plugins/Runtime/AndroidMoviePlayer/Source/AndroidMoviePlayer/Private/AndroidMovieStreamer.cpp b/Engine/Plugins/Runtime/AndroidMoviePlayer/Source/AndroidMoviePlayer/Private/AndroidMovieStreamer.cpp
--- a/Engine/Plugins/Runtime/AndroidMoviePlayer/Source/AndroidMoviePlayer/Private/AndroidMovieStreamer.cpp
+++ b/Engine/Plugins/Runtime/AndroidMoviePlayer/Source/AndroidMoviePlayer/Private/AndroidMovieStreamer.cpp
@@ -45,6 +45,35 @@ void FAndroidMediaPlayerStreamer::ForceCompletion()
 	CloseMovie();
 }
 
+static void DoUpdateTextureMovieSampleExecute(TWeakPtr<FJavaAndroidMediaPlayer, ESPMode::ThreadSafe> JavaMediaPlayerPtr, int32 DestTexture)
+{
+	auto PinnedJavaMediaPlayer = JavaMediaPlayerPtr.Pin();
+
+	if (!PinnedJavaMediaPlayer.IsValid())
+	{
+		return;
+	}
+
+	PinnedJavaMediaPlayer->GetVideoLastFrame(DestTexture);
+}
+
+struct FRHICommandUpdateTextureMovieSample final : public FRHICommand<FRHICommandUpdateTextureMovieSample>
+{
+	TWeakPtr<FJavaAndroidMediaPlayer, ESPMode::ThreadSafe> JavaMediaPlayerPtr;
+	int32 DestTexture;
+
+	FORCEINLINE_DEBUGGABLE FRHICommandUpdateTextureMovieSample(TWeakPtr<FJavaAndroidMediaPlayer, ESPMode::ThreadSafe> InJavaMediaPlayerPtr, int32 InDestTexture)
+		: JavaMediaPlayerPtr(InJavaMediaPlayerPtr)
+		, DestTexture(InDestTexture)
+	{
+	}
+	void Execute(FRHICommandListBase& CmdList)
+	{
+		QUICK_SCOPE_CYCLE_COUNTER(STAT_FRHICommandUpdateTextureMovieSample_Execute);
+		DoUpdateTextureMovieSampleExecute(JavaMediaPlayerPtr, DestTexture);
+	}
+};
+
 bool FAndroidMediaPlayerStreamer::Tick(float DeltaTime)
 {
 	// Check the list of textures pending deletion and remove any that are no longer valid
@@ -83,7 +112,15 @@ bool FAndroidMediaPlayerStreamer::Tick(float DeltaTime)
 			if (!FAndroidMisc::ShouldUseVulkan())
 			{
 				int32 DestTexture = *reinterpret_cast<int32*>(CurrentTexture->GetTypedResource().GetReference()->GetNativeResource());
-				bool frameSuccess = JavaMediaPlayer->GetVideoLastFrame(DestTexture);
+				if (IsRunningRHIInSeparateThread())
+				{
+					FRHICommandListImmediate &RHICommandList = GetImmediateCommandList_ForRenderCommand();
+					new (RHICommandList.AllocCommand<FRHICommandUpdateTextureMovieSample>()) FRHICommandUpdateTextureMovieSample(JavaMediaPlayer, DestTexture);
+				}
+				else
+				{
+					DoUpdateTextureMovieSampleExecute(JavaMediaPlayer, DestTexture);
+				}
 			}
 			else
 			{
diff --git a/Engine/Source/Runtime/Core/Public/Android/AndroidPlatform.h b/Engine/Source/Runtime/Core/Public/Android/AndroidPlatform.h
--- a/Engine/Source/Runtime/Core/Public/Android/AndroidPlatform.h
+++ b/Engine/Source/Runtime/Core/Public/Android/AndroidPlatform.h
@@ -54,7 +54,7 @@ typedef FAndroidTypes FPlatformTypes;
 #define PLATFORM_UI_HAS_MOBILE_SCROLLBARS			1
 #define PLATFORM_UI_NEEDS_TOOLTIPS					0
 #define PLATFORM_UI_NEEDS_FOCUS_OUTLINES			0
-#define PLATFORM_SUPPORTS_EARLY_MOVIE_PLAYBACK		1 // movies will start before engine is initalized
+#define PLATFORM_SUPPORTS_EARLY_MOVIE_PLAYBACK		0 // movies will start before engine is initalized
 
 #if defined(EXPERIMENTAL_OPENGL_RHITHREAD) && EXPERIMENTAL_OPENGL_RHITHREAD
 	#define PLATFORM_RHITHREAD_DEFAULT_BYPASS			0
diff --git a/Engine/Source/Runtime/Core/Public/IOS/IOSPlatform.h b/Engine/Source/Runtime/Core/Public/IOS/IOSPlatform.h
--- a/Engine/Source/Runtime/Core/Public/IOS/IOSPlatform.h
+++ b/Engine/Source/Runtime/Core/Public/IOS/IOSPlatform.h
@@ -47,7 +47,7 @@ typedef FIOSPlatformTypes FPlatformTypes;
 #define PLATFORM_ALLOW_NULL_RHI							1
 #define PLATFORM_ENABLE_VECTORINTRINSICS_NEON			PLATFORM_64BITS // disable vector intrinsics to make it compatible with 32-bit in Xcode 8.3
 #define PLATFORM_SUPPORTS_STACK_SYMBOLS					1
-#define PLATFORM_SUPPORTS_EARLY_MOVIE_PLAYBACK			1 // movies will start before engine is initalized
+#define PLATFORM_SUPPORTS_EARLY_MOVIE_PLAYBACK			0 // movies will start before engine is initalized
 
 // on iOS we now perform offline symbolication as it's significantly faster. Requires bGenerateCrashReportSymbols=true in the ini file.
 #define	PLATFORM_RUNTIME_MALLOCPROFILER_SYMBOLICATION	0	
