Fix TWeakPtr/TSharedPtr usage for virtual keyboard and startup movies on Android

#ue4
#android

diff --git a/Engine/Plugins/Runtime/AndroidMoviePlayer/Source/AndroidMoviePlayer/Private/AndroidMovieStreamer.h b/Engine/Plugins/Runtime/AndroidMoviePlayer/Source/AndroidMoviePlayer/Private/AndroidMovieStreamer.h
--- a/Engine/Plugins/Runtime/AndroidMoviePlayer/Source/AndroidMoviePlayer/Private/AndroidMovieStreamer.h
+++ b/Engine/Plugins/Runtime/AndroidMoviePlayer/Source/AndroidMoviePlayer/Private/AndroidMovieStreamer.h
@@ -49,7 +49,7 @@ private:
 	/**
 	Java side interface for MediaPlayer utility.
 	*/
-	TSharedPtr<FJavaAndroidMediaPlayer>    JavaMediaPlayer;
+	TSharedPtr<FJavaAndroidMediaPlayer, ESPMode::ThreadSafe> JavaMediaPlayer;
 
 	/**
 	Some state for the playing movie.
diff --git a/Engine/Source/Runtime/Launch/Private/Android/AndroidJNI.cpp b/Engine/Source/Runtime/Launch/Private/Android/AndroidJNI.cpp
--- a/Engine/Source/Runtime/Launch/Private/Android/AndroidJNI.cpp
+++ b/Engine/Source/Runtime/Launch/Private/Android/AndroidJNI.cpp
@@ -27,7 +27,7 @@ jclass GGameActivityClassID;
 jobject GGameActivityThis;
 
 // Pointer to target widget for virtual keyboard contents
-static TWeakPtr<IVirtualKeyboardEntry, ESPMode::ThreadSafe> VirtualKeyboardWidget;
+static TWeakPtr<IVirtualKeyboardEntry> VirtualKeyboardWidget;
 
 //virtualKeyboard shown
 static volatile bool GVirtualKeyboardShown = false;
