CIS fix

diff --git a/Engine/Plugins/Runtime/AndroidMoviePlayer/Source/AndroidMoviePlayer/Private/AndroidMovieStreamer.cpp b/Engine/Plugins/Runtime/AndroidMoviePlayer/Source/AndroidMoviePlayer/Private/AndroidMovieStreamer.cpp
--- a/Engine/Plugins/Runtime/AndroidMoviePlayer/Source/AndroidMoviePlayer/Private/AndroidMovieStreamer.cpp
+++ b/Engine/Plugins/Runtime/AndroidMoviePlayer/Source/AndroidMoviePlayer/Private/AndroidMovieStreamer.cpp
@@ -209,7 +209,7 @@ bool FAndroidMediaPlayerStreamer::StartNextMovie()
 		}
 
 		// is it a simple case (can just use file datasource)?
-		if (FileEntry.CompressionMethod == COMPRESS_None && !FileEntry.IsEncrypted())
+		if (FileEntry.CompressionMethodIndex == 0 && !FileEntry.IsEncrypted())
 		{
 			FString PakFilename = PakFile->GetFilename();
 			int64 PakHeaderSize = FileEntry.GetSerializedSize(PakFile->GetInfo().Version);
