Fixed: Ensure in ShooterGame on startup referencing VulkanDescriptorSets.h

diff --git a/Engine/Plugins/Runtime/AndroidMoviePlayer/Source/AndroidMoviePlayer/Private/AndroidMovieStreamer.cpp b/Engine/Plugins/Runtime/AndroidMoviePlayer/Source/AndroidMoviePlayer/Private/AndroidMovieStreamer.cpp
--- a/Engine/Plugins/Runtime/AndroidMoviePlayer/Source/AndroidMoviePlayer/Private/AndroidMovieStreamer.cpp
+++ b/Engine/Plugins/Runtime/AndroidMoviePlayer/Source/AndroidMoviePlayer/Private/AndroidMovieStreamer.cpp
@@ -278,7 +278,7 @@ bool FAndroidMediaPlayerStreamer::StartNextMovie()
 			VideoDimensions.X, VideoDimensions.Y,
 			PF_B8G8R8A8,
 			NULL,
-			TexCreate_RenderTargetable,
+			TexCreate_RenderTargetable | TexCreate_ShaderResource,
 			true));
 
 		uint32 FrameBytes = VideoDimensions.X * VideoDimensions.Y * GPixelFormats[PF_B8G8R8A8].BlockBytes;
@@ -294,6 +294,8 @@ bool FAndroidMediaPlayerStreamer::StartNextMovie()
 				void* TextureBuffer = RHILockTexture2D(TextureRHIRef->GetTypedResource(), 0, RLM_WriteOnly, Stride, false);
 				FMemory::Memset(TextureBuffer, 0, FrameBytes);
 				RHIUnlockTexture2D(TextureRHIRef->GetTypedResource(), 0, false);
+
+				RHICmdList.Transition(FRHITransitionInfo(TextureRHIRef->GetTypedResource(), ERHIAccess::Unknown, ERHIAccess::SRVMask));
 			});
 
 		MovieViewport->SetTexture(Texture);
