Add missing #include

#android

diff --git a/Engine/Plugins/Runtime/AndroidMoviePlayer/Source/AndroidMoviePlayer/Private/AndroidMovieStreamer.cpp b/Engine/Plugins/Runtime/AndroidMoviePlayer/Source/AndroidMoviePlayer/Private/AndroidMovieStreamer.cpp
--- a/Engine/Plugins/Runtime/AndroidMoviePlayer/Source/AndroidMoviePlayer/Private/AndroidMovieStreamer.cpp
+++ b/Engine/Plugins/Runtime/AndroidMoviePlayer/Source/AndroidMoviePlayer/Private/AndroidMovieStreamer.cpp
@@ -7,6 +7,7 @@
 #include "Android/AndroidPlatformFile.h"
 
 #include "Rendering/RenderingCommon.h"
+#include "RenderingThread.h"
 #include "RenderUtils.h"
 #include "Slate/SlateTextures.h"
 #include "MoviePlayer.h"
