Changing a number of places that use implicit command lists to instead use the one already available to them.

diff --git a/Engine/Plugins/Experimental/ImagePlate/Source/ImagePlate/Private/ImagePlateFileSequence.cpp b/Engine/Plugins/Experimental/ImagePlate/Source/ImagePlate/Private/ImagePlateFileSequence.cpp
--- a/Engine/Plugins/Experimental/ImagePlate/Source/ImagePlate/Private/ImagePlateFileSequence.cpp
+++ b/Engine/Plugins/Experimental/ImagePlate/Source/ImagePlate/Private/ImagePlateFileSequence.cpp
@@ -204,7 +204,7 @@ TFuture<void> FImagePlateSourceFrame::CopyTo(UTexture* DestinationTexture)
 			}
 
 			uint32 DestPitch = 0;
-			uint8* DestinationBuffer = (uint8*)RHILockTexture2D(Texture2DRHI, 0, RLM_WriteOnly, DestPitch, false);
+			uint8* DestinationBuffer = (uint8*)RHICmdList.LockTexture2D(Texture2DRHI, 0, RLM_WriteOnly, DestPitch, false);
 			if (!DestinationBuffer)
 			{
 				UE_LOG(LogImagePlateFileSequence, Warning, TEXT("Unable to lock texture for write"));
@@ -233,7 +233,7 @@ TFuture<void> FImagePlateSourceFrame::CopyTo(UTexture* DestinationTexture)
 				}
 			}
 
-			RHIUnlockTexture2D(Texture2DRHI, 0, false);
+			RHICmdList.UnlockTexture2D(Texture2DRHI, 0, false);
 
 			CommandData->Promise.SetValue();
 		}
diff --git a/Engine/Plugins/Runtime/AndroidMoviePlayer/Source/AndroidMoviePlayer/Private/AndroidMovieStreamer.cpp b/Engine/Plugins/Runtime/AndroidMoviePlayer/Source/AndroidMoviePlayer/Private/AndroidMovieStreamer.cpp
--- a/Engine/Plugins/Runtime/AndroidMoviePlayer/Source/AndroidMoviePlayer/Private/AndroidMovieStreamer.cpp
+++ b/Engine/Plugins/Runtime/AndroidMoviePlayer/Source/AndroidMoviePlayer/Private/AndroidMovieStreamer.cpp
@@ -105,8 +105,6 @@ bool FAndroidMediaPlayerStreamer::Tick(FRHICommandListBase& RHICmdList, float De
 
 	if (IsInRenderingThread())
 	{
-		FRHICommandListImmediate& RHICommandList = GetImmediateCommandList_ForRenderCommand();
-
 		int32 NextPosition = JavaMediaPlayer->GetCurrentPosition();
 		if (CurrentPosition != NextPosition)
 		{
@@ -115,14 +113,14 @@ bool FAndroidMediaPlayerStreamer::Tick(FRHICommandListBase& RHICmdList, float De
 
 			if (!CurrentTexture->IsInitialized())
 			{
-				CurrentTexture->InitResource(RHICommandList);
+				CurrentTexture->InitResource(RHICmdList);
 			}
 
 			if (!FAndroidMisc::ShouldUseVulkan())
 			{
 				if (IsRunningRHIInSeparateThread())
 				{
-					new (RHICommandList.AllocCommand<FRHICommandUpdateTextureMovieSample>()) FRHICommandUpdateTextureMovieSample(JavaMediaPlayer, CurrentTexture);
+					new (RHICmdList.AllocCommand<FRHICommandUpdateTextureMovieSample>()) FRHICommandUpdateTextureMovieSample(JavaMediaPlayer, CurrentTexture);
 				}
 				else
 				{
@@ -137,9 +135,9 @@ bool FAndroidMediaPlayerStreamer::Tick(FRHICommandListBase& RHICmdList, float De
 				void* LastFrameData = nullptr;
 				if(JavaMediaPlayer->GetVideoLastFrameData(LastFrameData, SampleCount))
 				{
-					void* DestTextureData = RHILockTexture2D(CurrentTexture->GetTypedResource(), 0, RLM_WriteOnly, Stride, false);
+					void* DestTextureData = RHICmdList.GetAsImmediate().LockTexture2D(CurrentTexture->GetTypedResource(), 0, RLM_WriteOnly, Stride, false);
 					FMemory::Memcpy(DestTextureData, LastFrameData, SampleCount);
-					RHIUnlockTexture2D(CurrentTexture->GetTypedResource(), 0, false);
+					RHICmdList.GetAsImmediate().UnlockTexture2D(CurrentTexture->GetTypedResource(), 0, false);
 				}
 			}
 
@@ -293,9 +291,9 @@ bool FAndroidMediaPlayerStreamer::StartNextMovie()
 	
 				// clear texture to black
 				uint32 Stride = 0;
-				void* TextureBuffer = RHILockTexture2D(TextureRHIRef->GetTypedResource(), 0, RLM_WriteOnly, Stride, false);
+				void* TextureBuffer = RHICmdList.LockTexture2D(TextureRHIRef->GetTypedResource(), 0, RLM_WriteOnly, Stride, false);
 				FMemory::Memset(TextureBuffer, 0, FrameBytes);
-				RHIUnlockTexture2D(TextureRHIRef->GetTypedResource(), 0, false);
+				RHICmdList.UnlockTexture2D(TextureRHIRef->GetTypedResource(), 0, false);
 
 				RHICmdList.Transition(FRHITransitionInfo(TextureRHIRef->GetTypedResource(), ERHIAccess::Unknown, ERHIAccess::SRVMask));
 			});
diff --git a/Engine/Plugins/Tests/RHITests/Source/RHITests/Private/RHIBufferTests.cpp b/Engine/Plugins/Tests/RHITests/Source/RHITests/Private/RHIBufferTests.cpp
--- a/Engine/Plugins/Tests/RHITests/Source/RHITests/Private/RHIBufferTests.cpp
+++ b/Engine/Plugins/Tests/RHITests/Source/RHITests/Private/RHIBufferTests.cpp
@@ -29,9 +29,9 @@ bool FRHIBufferTests::VerifyBufferContents(const TCHAR* TestName, FRHICommandLis
 		FRHIStagingBuffer* StagingBuffer = StagingBuffers[Index];
 		uint32 NumBytes = Buffer->GetSize();
 
-		void* Memory = RHILockStagingBuffer(StagingBuffer, 0, NumBytes);
+		void* Memory = RHICmdList.LockStagingBuffer(StagingBuffer, 0, NumBytes);
 		bool Result = VerifyCallback(Index, Memory, NumBytes);
-		RHIUnlockStagingBuffer(StagingBuffer);
+		RHICmdList.UnlockStagingBuffer(StagingBuffer);
 
 		if (!Result)
 		{
@@ -78,7 +78,7 @@ bool FRHIBufferTests::VerifyBufferContents(const TCHAR* TestName, FRHICommandLis
 		RHICmdList.BlockUntilGPUIdle();
 		RHICmdList.FlushResources();
 
-		void* Memory = RHILockStagingBuffer(StagingBuffer, 0, NumBytes);
+		void* Memory = RHICmdList.LockStagingBuffer(StagingBuffer, 0, NumBytes);
 
 		check(ExpectedData.Num() == NumBytes);
 		Result = FMemory::Memcmp(ExpectedData.GetData(), Memory, NumBytes) == 0;
@@ -95,7 +95,7 @@ bool FRHIBufferTests::VerifyBufferContents(const TCHAR* TestName, FRHICommandLis
 			UE_LOG(LogRHIUnitTestCommandlet, Display, TEXT("Test passed. \"%s\""), TestName);
 		}
 
-		RHIUnlockStagingBuffer(StagingBuffer);
+		RHICmdList.UnlockStagingBuffer(StagingBuffer);
 	}
 
 	// Immediate flush to clean up the staging buffer / other resources
diff --git a/Engine/Source/Runtime/Engine/Private/Slate/SlateTextures.cpp b/Engine/Source/Runtime/Engine/Private/Slate/SlateTextures.cpp
--- a/Engine/Source/Runtime/Engine/Private/Slate/SlateTextures.cpp
+++ b/Engine/Source/Runtime/Engine/Private/Slate/SlateTextures.cpp
@@ -190,10 +190,10 @@ void FSlateTexture2DRHIRef::ResizeTexture(uint32 InWidth, uint32 InHeight)
 	}
 }
 
-void FSlateTexture2DRHIRef::SetTextureData(const TArray<uint8>& Bytes)
+void FSlateTexture2DRHIRef::SetTextureData(FRHICommandListImmediate& RHICmdList, const TArray<uint8>& Bytes)
 {
 	uint32 DstStride = 0;
-	uint8* DstData = (uint8*) RHILockTexture2D(GetTypedResource(), 0, RLM_WriteOnly, DstStride, false);
+	uint8* DstData = (uint8*)RHICmdList.LockTexture2D(GetTypedResource(), 0, RLM_WriteOnly, DstStride, false);
 
 	const uint32 NumBlocksX = (Width  + GPixelFormats[PixelFormat].BlockSizeX - 1) / GPixelFormats[PixelFormat].BlockSizeX;
 	const uint32 NumBlocksY = (Height + GPixelFormats[PixelFormat].BlockSizeY - 1) / GPixelFormats[PixelFormat].BlockSizeY;
@@ -215,14 +215,14 @@ void FSlateTexture2DRHIRef::SetTextureData(const TArray<uint8>& Bytes)
 		}
 	}
 
-	RHIUnlockTexture2D(GetTypedResource(), 0, false);
+	RHICmdList.UnlockTexture2D(GetTypedResource(), 0, false);
 }
 
 void FSlateTexture2DRHIRef::UpdateTexture(const TArray<uint8>& Bytes)
 {
 	if (IsInRenderingThread())
 	{
-		SetTextureData(Bytes);
+		SetTextureData(FRHICommandListImmediate::Get(), Bytes);
 	}
 	else
 	{
@@ -230,7 +230,7 @@ void FSlateTexture2DRHIRef::UpdateTexture(const TArray<uint8>& Bytes)
 		ENQUEUE_RENDER_COMMAND(UpdateSlateTexture)(
 			[TextureRHIRef, Bytes](FRHICommandListImmediate& RHICmdList)
 			{
-				TextureRHIRef->SetTextureData(Bytes);
+				TextureRHIRef->SetTextureData(RHICmdList, Bytes);
 			});
 	}
 }
diff --git a/Engine/Source/Runtime/Engine/Private/TextureCube.cpp b/Engine/Source/Runtime/Engine/Private/TextureCube.cpp
--- a/Engine/Source/Runtime/Engine/Private/TextureCube.cpp
+++ b/Engine/Source/Runtime/Engine/Private/TextureCube.cpp
@@ -483,9 +483,9 @@ public:
 				if( MipData[FaceIndex][MipIndex] != NULL )
 				{
 					uint32 DestStride;
-					void* TheMipData = RHILockTextureCubeFace( TextureCubeRHI, FaceIndex, 0, MipIndex, RLM_WriteOnly, DestStride, false );
+					void* TheMipData = RHICmdList.GetAsImmediate().LockTextureCubeFace( TextureCubeRHI, FaceIndex, 0, MipIndex, RLM_WriteOnly, DestStride, false );
 					GetData( FaceIndex, MipIndex, TheMipData, DestStride );
-					RHIUnlockTextureCubeFace( TextureCubeRHI, FaceIndex, 0, MipIndex, false );
+					RHICmdList.GetAsImmediate().UnlockTextureCubeFace( TextureCubeRHI, FaceIndex, 0, MipIndex, false );
 				}
 			}
 		}
diff --git a/Engine/Source/Runtime/Engine/Private/TextureCubeArray.cpp b/Engine/Source/Runtime/Engine/Private/TextureCubeArray.cpp
--- a/Engine/Source/Runtime/Engine/Private/TextureCubeArray.cpp
+++ b/Engine/Source/Runtime/Engine/Private/TextureCubeArray.cpp
@@ -159,9 +159,9 @@ public:
 				for (uint32 FaceIndex = 0u; FaceIndex < 6u; ++FaceIndex)
 				{
 					uint32 DestStride;
-					void* TheMipData = RHILockTextureCubeFace(TextureCubeRHI, FaceIndex, ArrayIndex, MipIndex, RLM_WriteOnly, DestStride, false);
+					void* TheMipData = RHICmdList.GetAsImmediate().LockTextureCubeFace(TextureCubeRHI, FaceIndex, ArrayIndex, MipIndex, RLM_WriteOnly, DestStride, false);
 					GetData(ArrayIndex, FaceIndex, MipIndex, TheMipData, DestStride);
-					RHIUnlockTextureCubeFace(TextureCubeRHI, FaceIndex, ArrayIndex, MipIndex, false);
+					RHICmdList.GetAsImmediate().UnlockTextureCubeFace(TextureCubeRHI, FaceIndex, ArrayIndex, MipIndex, false);
 				}
 			}
 		}
diff --git a/Engine/Source/Runtime/Engine/Public/Slate/SlateTextures.h b/Engine/Source/Runtime/Engine/Public/Slate/SlateTextures.h
--- a/Engine/Source/Runtime/Engine/Public/Slate/SlateTextures.h
+++ b/Engine/Source/Runtime/Engine/Public/Slate/SlateTextures.h
@@ -92,7 +92,7 @@ protected:
 	/** Height of this texture */
 	uint32 Height;
 private:
-	ENGINE_API void SetTextureData(const TArray<uint8>& Bytes);
+	ENGINE_API void SetTextureData(FRHICommandListImmediate& RHICmdList, const TArray<uint8>& Bytes);
 
 	/** Texture creation flags for if this texture needs to be recreated dynamically */
 	ETextureCreateFlags TexCreateFlags;
diff --git a/Engine/Source/Runtime/Renderer/Private/ReflectionEnvironmentCapture.cpp b/Engine/Source/Runtime/Renderer/Private/ReflectionEnvironmentCapture.cpp
--- a/Engine/Source/Runtime/Renderer/Private/ReflectionEnvironmentCapture.cpp
+++ b/Engine/Source/Runtime/Renderer/Private/ReflectionEnvironmentCapture.cpp
@@ -1136,7 +1136,7 @@ void FScene::GetReflectionCaptureData(UReflectionCaptureComponent* Component, FR
 	FlushRenderingCommands();
 }
 
-void UploadReflectionCapture_RenderingThread(FScene* Scene, const FReflectionCaptureData* CaptureData, const UReflectionCaptureComponent* CaptureComponent)
+void UploadReflectionCapture_RenderingThread(FRHICommandListImmediate& RHICmdList, FScene* Scene, const FReflectionCaptureData* CaptureData, const UReflectionCaptureComponent* CaptureComponent)
 {
 	// Due to memory limitations, it's possible for the in-memory size to be smaller than the originally captured size.
 	const int32 EffectiveTopMipSize = Scene->ReflectionSceneData.CubemapArray.GetCubemapSize();
@@ -1168,7 +1168,7 @@ void UploadReflectionCapture_RenderingThread(FScene* Scene, const FReflectionCap
 		for (int32 CubeFace = 0; CubeFace < CubeFace_MAX; CubeFace++)
 		{
 			uint32 DestStride = 0;
-			uint8* DestBuffer = (uint8*)RHILockTextureCubeFace(CubeMapArray, CubeFace, CaptureIndex, MipIndex, RLM_WriteOnly, DestStride, false);
+			uint8* DestBuffer = (uint8*)RHICmdList.LockTextureCubeFace(CubeMapArray, CubeFace, CaptureIndex, MipIndex, RLM_WriteOnly, DestStride, false);
 
 			// Handle DestStride by copying each row
 			for (int32 Y = 0; Y < MipSize; Y++)
@@ -1179,7 +1179,7 @@ void UploadReflectionCapture_RenderingThread(FScene* Scene, const FReflectionCap
 				FMemory::Memcpy(DestPtr, SourcePtr, MipSize * sizeof(FFloat16Color));
 			}
 
-			RHIUnlockTextureCubeFace(CubeMapArray, CubeFace, CaptureIndex, MipIndex, false);
+			RHICmdList.UnlockTextureCubeFace(CubeMapArray, CubeFace, CaptureIndex, MipIndex, false);
 		}
 
 		MipBaseIndex += CubeFaceBytes * CubeFace_MAX;
@@ -1463,7 +1463,7 @@ void FScene::CaptureOrUploadReflectionCapture(UReflectionCaptureComponent* Captu
 				// This means switching between LightingScenarios only works if the scenario level is reloaded (not simply made hidden / visible again)
 				if (!CaptureData->HasBeenUploadedFinal())
 				{
-					UploadReflectionCapture_RenderingThread(Scene, CaptureData, CaptureComponent);
+					UploadReflectionCapture_RenderingThread(RHICmdList, Scene, CaptureData, CaptureComponent);
 
 					CaptureData->OnDataUploadedToGPUFinal();
 				}
diff --git a/Engine/Source/Runtime/SlateRHIRenderer/Private/SlateRHITextureAtlas.cpp b/Engine/Source/Runtime/SlateRHIRenderer/Private/SlateRHITextureAtlas.cpp
--- a/Engine/Source/Runtime/SlateRHIRenderer/Private/SlateRHITextureAtlas.cpp
+++ b/Engine/Source/Runtime/SlateRHIRenderer/Private/SlateRHITextureAtlas.cpp
@@ -41,13 +41,13 @@ void FSlateTextureAtlasRHI::ReleaseResources()
 	BeginReleaseResource(AtlasTexture);
 }
 
-void FSlateTextureAtlasRHI::UpdateTexture_RenderThread( FSlateTextureData* RenderThreadData )
+void FSlateTextureAtlasRHI::UpdateTexture_RenderThread(FRHICommandListImmediate& RHICmdList, FSlateTextureData* RenderThreadData)
 {
 	check(IsInRenderingThread());
 
 	if (!AtlasTexture->IsInitialized())
 	{
-		AtlasTexture->InitResource(FRHICommandListImmediate::Get());
+		AtlasTexture->InitResource(RHICmdList);
 	}
 
 	check(AtlasTexture->IsInitialized());
@@ -58,14 +58,14 @@ void FSlateTextureAtlasRHI::UpdateTexture_RenderThread( FSlateTextureData* Rende
 	if ( RawBytesPtr != nullptr )
 	{
 		uint32 DestStride;
-		uint8* TempData = (uint8*)RHILockTexture2D(AtlasTexture->GetTypedResource(), 0, RLM_WriteOnly, /*out*/ DestStride, false);
+		uint8* TempData = (uint8*)RHICmdList.LockTexture2D(AtlasTexture->GetTypedResource(), 0, RLM_WriteOnly, /*out*/ DestStride, false);
 		// check(DestStride == (RenderThreadData->GetBytesPerPixel() * RenderThreadData->GetWidth())); // Temporarily disable check
 
 		if ( TempData != nullptr )
 		{
 			FMemory::Memcpy(TempData, RawBytesPtr, RenderThreadData->GetBytesPerPixel() * RenderThreadData->GetWidth() * RenderThreadData->GetHeight());
 
-			RHIUnlockTexture2D(AtlasTexture->GetTypedResource(), 0, false);
+			RHICmdList.UnlockTexture2D(AtlasTexture->GetTypedResource(), 0, false);
 		}
 	}
 
@@ -88,7 +88,7 @@ void FSlateTextureAtlasRHI::ConditionalUpdateTexture( )
 		ENQUEUE_RENDER_COMMAND(SlateUpdateAtlasTextureCommand)(
 			[Atlas, RenderThreadData](FRHICommandListImmediate& RHICmdList)
 			{
-				Atlas->UpdateTexture_RenderThread( RenderThreadData );
+				Atlas->UpdateTexture_RenderThread(RHICmdList, RenderThreadData);
 			});
 
 		bNeedsUpdate = false;
diff --git a/Engine/Source/Runtime/SlateRHIRenderer/Private/SlateRHITextureAtlas.h b/Engine/Source/Runtime/SlateRHIRenderer/Private/SlateRHITextureAtlas.h
--- a/Engine/Source/Runtime/SlateRHIRenderer/Private/SlateRHITextureAtlas.h
+++ b/Engine/Source/Runtime/SlateRHIRenderer/Private/SlateRHITextureAtlas.h
@@ -5,6 +5,7 @@
 #include "CoreMinimal.h"
 #include "Textures/TextureAtlas.h"
 
+class FRHICommandListImmediate;
 class FSlateTexture2DRHIRef;
 struct FSlateTextureData;
 
@@ -44,9 +45,10 @@ public:
 	/**
 	 * Updates the texture on the render thread.
 	 *
+	 * @param RHICmdList
 	 * @param RenderThreadData
 	 */
-	void UpdateTexture_RenderThread( FSlateTextureData* RenderThreadData );
+	void UpdateTexture_RenderThread(FRHICommandListImmediate& RHICmdList, FSlateTextureData* RenderThreadData);
 
 public:
 
